Add WebTestsServerCpp: C++ UE5 HTTP test server replacing Python WebTestsServer
- Rename WebTestsServer to WebTestsServerPy
- Add WebTestsServerCpp program using FHttpServerModule
- Add HEAD/TRACE/CONNECT verb support to HTTPServer module
- Add streaming body queue to FHttpServerResponse
- Store FullPath on FHttpServerRequest before MakeRelative

#Synced-CL 52039227

diff --git a/Engine/Source/Programs/WebTests/Private/TestHttp.cpp b/Engine/Source/Programs/WebTests/Private/TestHttp.cpp
--- a/Engine/Source/Programs/WebTests/Private/TestHttp.cpp
+++ b/Engine/Source/Programs/WebTests/Private/TestHttp.cpp
@@ -84,10 +84,10 @@ namespace HttpTestUtils
 class FHttpTestLogLevelInitializer
 {
 public:
-	FHttpTestLogLevelInitializer()
+	FHttpTestLogLevelInitializer(bool bInVeryVerbose)
 		: OldVerbosity(LogHttp.GetVerbosity())
+		, bVeryVerbose(bInVeryVerbose)
 	{
-		FParse::Bool(FCommandLine::Get(), TEXT("very_verbose="), bVeryVerbose);
 		if (bVeryVerbose)
 		{
 			LogHttp.SetVerbosity(ELogVerbosity::VeryVerbose);
@@ -118,8 +118,8 @@ public:
 		}
 	}
 
-	bool bVeryVerbose = false;
 	ELogVerbosity::Type OldVerbosity;
+	bool bVeryVerbose = false;
 };
 
 class FMockRetryManager : public FHttpRetrySystem::FManager
@@ -201,9 +201,9 @@ public:
 		return bRetryEnabled ? HttpRetryManager->CreateRequest() : HttpModule->CreateRequest();
 	}
 
-	void SetInvalidDomainToTestConnectTimeout() { ServerSettings.WebServerIp = TEXT("10.255.255.1"); ServerSettings.WebServerPort = 8765; } // non-routable IP address with a random port
-	const FString UrlDomain() const { return FString::Format(TEXT("{0}:{1}"), { *ServerSettings.WebServerIp, ServerSettings.WebServerPort }); }
-	const FString UrlHttpTests() const { return ServerSettings.UrlHttpTests(); }
+	void SetInvalidDomainToTestConnectTimeout() { Settings.WebServerIp = TEXT("10.255.255.1"); Settings.WebServerHttpPort = 8765; } // non-routable IP address with a random port
+	const FString UrlDomain() const { return FString::Format(TEXT("{0}:{1}"), { *Settings.WebServerIp, Settings.WebServerHttpPort }); }
+	const FString UrlHttpTests() const { return Settings.UrlHttpTests(); }
 	const FString UrlToTestMethods() const { return FString::Format(TEXT("{0}/methods"), { *UrlHttpTests() }); }
 	const FString UrlStreamDownload(uint32 Chunks, uint32 ChunkSize, uint32 ChunkLatency=0) { return FString::Format(TEXT("{0}/streaming_download/{1}/{2}/{3}/"), { *UrlHttpTests(), Chunks, ChunkSize, ChunkLatency }); }
 	const FString UrlStreamUpload() { return FString::Format(TEXT("{0}/streaming_upload_put"), { *UrlHttpTests() }); }
@@ -212,13 +212,13 @@ public:
 	const FString UrlUnixSocketHttpTests() const { return "http://localhost/webtests/unixsockettests"; }
 	const FString UrlPostEchoBody() const { return FString::Format(TEXT("{0}/post_echo_body"), { *UrlHttpTests() }); }
 
-	FWebTestServerSettings ServerSettings;
+	FWebTestSettings Settings;
 	FString WebServerUnixSocket;
 	FMockHttpModule* HttpModule = nullptr;
 	bool bRunHeavyTests;
 	bool bRetryEnabled;
 	bool bEventLoopEnabled;
-	FHttpTestLogLevelInitializer HttpTestLogLevelInitializer;
+	FHttpTestLogLevelInitializer HttpTestLogLevelInitializer { Settings.bVeryVerbose };
 	TSharedPtr<FMockRetryManager> HttpRetryManager;
 };
 
diff --git a/Engine/Source/Programs/WebTests/Private/TestWebSockets.cpp b/Engine/Source/Programs/WebTests/Private/TestWebSockets.cpp
--- a/Engine/Source/Programs/WebTests/Private/TestWebSockets.cpp
+++ b/Engine/Source/Programs/WebTests/Private/TestWebSockets.cpp
@@ -33,6 +33,11 @@ public:
 	{
 		CVarHttpInsecureProtocolEnabled->Set(true);
 
+		if (Settings.bVeryVerbose)
+		{
+			LogWebSockets.SetVerbosity(ELogVerbosity::VeryVerbose);
+		}
+
 		WebSocketsModule = new FWebSocketsModule();
 		IModuleInterface* Module = WebSocketsModule;
 		Module->StartupModule();
@@ -55,13 +60,16 @@ public:
 
 	void DisableWarningsInThisTest()
 	{
-		LogWebSockets.SetVerbosity(ELogVerbosity::Error);
+		if (!Settings.bVeryVerbose)
+		{
+			LogWebSockets.SetVerbosity(ELogVerbosity::Error);
+		}
 	}
 
-	const FString UrlWithInvalidPortToTestConnectTimeout() const { return FString::Format(TEXT("ws://{0}:{1}"), { *ServerSettings.WebServerIp, 8765 }); }
-	const FString UrlWebSocketsTests() const { return ServerSettings.UrlWebSocketsTests(); }
+	const FString UrlWithInvalidPortToTestConnectTimeout() const { return FString::Format(TEXT("ws://{0}:{1}"), { *Settings.WebServerIp, 8765 }); }
+	const FString UrlWebSocketsTests() const { return Settings.UrlWebSocketsTests(); }
 
-	FWebTestServerSettings ServerSettings;
+	FWebTestSettings Settings;
 	FWebSocketsModule* WebSocketsModule;
 	ELogVerbosity::Type OldVerbosity;
 };
@@ -242,7 +250,7 @@ TEST_CASE_METHOD(FHttpAndWebSocketsFixture, "WebSockets works correctly when Htt
 
 	// Fire an HTTP request via FHttpModule while WebSockets is also active.
 	TSharedRef<IHttpRequest, ESPMode::ThreadSafe> HttpRequest = HttpModule->CreateRequest();
-	HttpRequest->SetURL(FString::Format(TEXT("{0}/methods"), { *ServerSettings.UrlHttpTests() }));
+	HttpRequest->SetURL(FString::Format(TEXT("{0}/methods"), { *Settings.UrlHttpTests() }));
 	HttpRequest->SetVerb(TEXT("GET"));
 	HttpRequest->OnProcessRequestComplete().BindLambda([bHttpDone, TryFinish](FHttpRequestPtr /* Request */, FHttpResponsePtr /* Response */, bool bProcessedSuccessfully)
 	{
@@ -316,14 +324,14 @@ public:
 		}
 	}
 
-	const FString UrlWebSocketsTests() const { return ServerSettings.UrlWebSocketsTests(); }
+	const FString UrlWebSocketsTests() const { return Settings.UrlWebSocketsTests(); }
 
 	void RequestQuit()          { TickRunner.bQuitRequested = true; }
 	void SetSucceeded()         { TickRunner.bSucceeded = true; }
 	bool HasSucceeded() const   { return TickRunner.bSucceeded; }
 
 	FWebTestTickRunner TickRunner;
-	FWebTestServerSettings ServerSettings;
+	FWebTestSettings Settings;
 	FHttpModule* HttpModule = nullptr;
 	FWebSocketsModule* WebSocketsModule = nullptr;
 	ELogVerbosity::Type OldVerbosity;
@@ -346,7 +354,7 @@ TEST_CASE_METHOD(FHttpFirstThenWebSocketsFixture, "WebSockets works correctly wh
 
 	// Fire an HTTP request via FHttpModule while WebSockets is also active.
 	TSharedRef<IHttpRequest, ESPMode::ThreadSafe> HttpRequest = HttpModule->CreateRequest();
-	HttpRequest->SetURL(FString::Format(TEXT("{0}/methods"), { *ServerSettings.UrlHttpTests() }));
+	HttpRequest->SetURL(FString::Format(TEXT("{0}/methods"), { *Settings.UrlHttpTests() }));
 	HttpRequest->SetVerb(TEXT("GET"));
 	HttpRequest->OnProcessRequestComplete().BindLambda([bHttpDone, TryFinish](FHttpRequestPtr /* Request */, FHttpResponsePtr /* Response */, bool bProcessedSuccessfully)
 	{
diff --git a/Engine/Source/Programs/WebTests/Private/WebTestFixtures.h b/Engine/Source/Programs/WebTests/Private/WebTestFixtures.h
--- a/Engine/Source/Programs/WebTests/Private/WebTestFixtures.h
+++ b/Engine/Source/Programs/WebTests/Private/WebTestFixtures.h
@@ -15,7 +15,7 @@
  *
  *  Contains:
  *  - CVarHttpInsecureProtocolEnabled extern (defined in HttpModule.cpp)
- *  - FWebTestServerSettings: command-line driven server address helpers
+ *  - FWebTestSettings: command-line driven server address and verbosity helpers
  *  - FWebTestTickRunner: engine tick simulation with quit-requested loop
  *
  *  -----------------------------------------------------------------------
@@ -25,42 +25,45 @@
 extern TAutoConsoleVariable<bool> CVarHttpInsecureProtocolEnabled;
 
 /**
- * Parses web test server connection settings from the command line.
- * Both Http and WebSocket tests target the same web server process.
+ * Parses web test runtime settings from the command line.
+ * Http and WebSocket tests may use different ports on the same server host.
  *
  * Command-line parameters:
  *   -web_server_ip=<ip>                    (default: 127.0.0.1)
- *   -web_server_websockets_port=<port>      (default: 8000)
- *   -web_server_http_port=<port>            (default: same as websockets port)
+ *   -web_server_http_port=<port>            (default: 8000)
+ *   -web_server_websockets_port=<port>      (default: 8001)
+ *   -very_verbose=true                      (default: false) — enables VeryVerbose log level
  */
-struct FWebTestServerSettings
+struct FWebTestSettings
 {
-	FWebTestServerSettings()
+	FWebTestSettings()
 		: WebServerIp(TEXT("127.0.0.1"))
-		, WebServerPort(8000)
+		, WebServerHttpPort(8000)
+		, WebServerWsPort(8001)
+		, bVeryVerbose(false)
 	{
 		FParse::Value(FCommandLine::Get(), TEXT("web_server_ip="), WebServerIp);
-		// Accept either param name; websockets tests historically used web_server_websockets_port.
-		if (!FParse::Value(FCommandLine::Get(), TEXT("web_server_http_port="), WebServerPort))
-		{
-			FParse::Value(FCommandLine::Get(), TEXT("web_server_websockets_port="), WebServerPort);
-		}
+		FParse::Value(FCommandLine::Get(), TEXT("web_server_http_port="), WebServerHttpPort);
+		FParse::Value(FCommandLine::Get(), TEXT("web_server_websockets_port="), WebServerWsPort);
+		FParse::Bool(FCommandLine::Get(), TEXT("very_verbose="), bVeryVerbose);
 	}
 
 	/** Base HTTP URL: "http://<ip>:<port>" */
-	FString HttpUrlBase() const { return FString::Format(TEXT("http://{0}:{1}"), { *WebServerIp, WebServerPort }); }
+	FString HttpUrlBase() const { return FString::Format(TEXT("http://{0}:{1}"), { *WebServerIp, WebServerHttpPort }); }
 
 	/** URL prefix for the http test routes: "http://<ip>:<port>/webtests/httptests" */
 	FString UrlHttpTests() const { return FString::Format(TEXT("{0}/webtests/httptests"), { *HttpUrlBase() }); }
 
-	/** Base WebSocket URL: "ws://<ip>:<port>" */
-	FString WsUrlBase() const { return FString::Format(TEXT("ws://{0}:{1}"), { *WebServerIp, WebServerPort }); }
+	/** Base WebSocket URL: "ws://<ip>:<wsport>" */
+	FString WsUrlBase() const { return FString::Format(TEXT("ws://{0}:{1}"), { *WebServerIp, WebServerWsPort }); }
 
-	/** URL prefix for the websocket test routes: "ws://<ip>:<port>/webtests/websocketstests" */
+	/** URL prefix for the websocket test routes: "ws://<ip>:<wsport>/webtests/websocketstests" */
 	FString UrlWebSocketsTests() const { return FString::Format(TEXT("{0}/webtests/websocketstests"), { *WsUrlBase() }); }
 
 	FString WebServerIp;
-	uint32  WebServerPort;
+	uint32  WebServerHttpPort;
+	uint32  WebServerWsPort;
+	bool    bVeryVerbose;
 };
 
 /**
diff --git a/Engine/Source/Programs/WebTests/README.md b/Engine/Source/Programs/WebTests/README.md
--- a/Engine/Source/Programs/WebTests/README.md
+++ b/Engine/Source/Programs/WebTests/README.md
@@ -5,18 +5,84 @@ This is a collection of tests for web features like http, websocket, ssl etc.
 
 # Running the Tests
 
-## Run the WebTestsServer(NOTE: This server is created for TEST PURPOSE ONLY, it's not production ready!!!):
-	In ../WebTestsServer folder:
+## Start the test server (NOTE: for TEST PURPOSE ONLY, not production ready)
+
+Two server options are available. Both listen on HTTP port 8000 and WebSocket port 8001.
+
+### Option A: C++ server (recommended — no Python required)
+
+Windows:
+	In ../WebTestsServerCpp folder:
+		> runserver.bat
+
+Mac/Linux:
+	In ../WebTestsServerCpp folder:
+		> sh runserver.sh
+
+### Option B: Python server (Django Channels / Daphne)
+
+Windows:
+	In ../WebTestsServerPy folder:
 		> runserver.bat
-	Or if you use docker:
-		> dockerbuildandrun.bat
+
+Mac/Linux:
+	In ../WebTestsServerPy folder:
+		> sh runserver.sh
 
 ## Run the tests from VS:
-	Set `WebTests` as the startup project and set Solution Configuration to a `Development`.
-	If running tests on other devices, pass in the ip as command line args, after extra args AT THE END, like: "--extra-args --web_server_ip=your.pc.ip.address"
+	Set `WebTests` as the startup project and set Solution Configuration to `Development`.
+	If running tests on other devices, pass in the ip as command line args, after extra args AT THE END, like: "--extra-args --web_server_ip=your.pc.ip.address --log"
 	Compile and debug
 
-# Adding new test case in WebTestsServer:
-	Add/change the code in ./WebTestsServer code, most likely in httptests/urls.py and httptests/views.py, and save. Code will be reloaded if the web server is running
-	When using docker, also need to run dockerbuildandrun.bat
-	For more info about how to code in django, check https://docs.djangoproject.com/en/
+## Compile and run the tests from command line:
+To compile:
+	.\Engine\Build\BatchFiles\RunUBT.bat WebTests Win64 Development -Progress
+
+To compile for any other platform, replace "Win64" in the param
+
+It uses Catch2 test framework, so in order to run with specific case or tag, pass the filter as first param, for example:
+	.\Engine\Binaries\Win64\WebTests\WebTests.exe "Get domain of url can work well"
+	.\Engine\Binaries\Win64\WebTests\WebTests.exe "[WebSockets]"
+	.\Engine\Binaries\Win64\WebTests\WebTests.exe "[HTTP]"
+
+# Adding new test cases in WebTestsServerPy (Python):
+	Add/change the code in ./WebTestsServerPy, most likely in httptests/urls.py and httptests/views.py, and save. Code will be reloaded if the web server is running.
+	For more info about how to code in Django, see https://docs.djangoproject.com/en/
+
+# Adding new test cases in WebTestsServerCpp (C++):
+
+After editing either file below, rebuild WebTestsServerCpp and restart the server.
+
+## New WebSocket route (WebTestsServerCpp/Private/WebTestsWebSocketServer.cpp)
+
+Add a new `else if` branch in the `OnMessage` handler inside `StartWebSocketServer()`:
+
+	else if (Path.EndsWith(TEXT("/your_route/")))
+	{
+	    // echo with modification:
+	    Conn->SendText(FString::Printf(TEXT("pong: %s"), *Msg));
+	    // or close the connection:
+	    // Conn->Close(1000);
+	}
+
+The path matches the URL the client connects to, e.g. ws://localhost:8001/webtests/websocketstests/your_route/
+
+## New HTTP route (WebTestsServerCpp/Private/WebTestsHttpServer.cpp)
+
+1. Write a handler function above StartHttpServer():
+
+	static bool HandleYourRoute(const FHttpServerRequest& Request, const FHttpResultCallback& OnComplete)
+	{
+	    LogRequest(Request, 200);
+	    OnComplete(JsonResponse(TEXT("{\"ok\":true}")));
+	    return true;
+	}
+
+2. Register it in StartHttpServer() with Router->BindRoute:
+
+	Router->BindRoute(FHttpPath(TEXT("/webtests/httptests/your_route")),
+	    EVerb::VERB_GET,
+	    FHttpRequestHandler::CreateStatic(&HandleYourRoute));
+
+Use :param_name segments for path parameters (e.g. /your_route/:id/); read them with Request.PathParams.Find(TEXT("id")).
+Query parameters are in Request.QueryParams.
diff --git a/Engine/Source/Programs/WebTestsServerCpp/Private/WebTestsHttpServer.cpp b/Engine/Source/Programs/WebTestsServerCpp/Private/WebTestsHttpServer.cpp
new file mode 100644
--- /dev/null
+++ b/Engine/Source/Programs/WebTestsServerCpp/Private/WebTestsHttpServer.cpp
@@ -0,0 +1,395 @@
+// Copyright Epic Games, Inc. All Rights Reserved.
+
+#include "WebTestsServer.h"
+#include "HAL/PlatformProcess.h"
+#include "HAL/Runnable.h"
+#include "HAL/RunnableThread.h"
+#include "Templates/Atomic.h"
+#include "HttpPath.h"
+#include "HttpRequestHandler.h"
+#include "HttpRouteHandle.h"
+#include "HttpServerConstants.h"
+#include "HttpServerModule.h"
+#include "HttpServerRequest.h"
+#include "HttpServerResponse.h"
+#include "IHttpRouter.h"
+#include "Misc/Compression.h"
+
+DEFINE_LOG_CATEGORY_STATIC(LogHttpServer, Log, All);
+
+// ---------------------------------------------------------------------------
+// Helper: build a JSON response with Content-Type application/json
+// ---------------------------------------------------------------------------
+static TUniquePtr<FHttpServerResponse> JsonResponse(const FString& JsonBody, EHttpServerResponseCodes Code = EHttpServerResponseCodes::Ok)
+{
+	FTCHARToUTF8 Utf8(*JsonBody);
+	TArray<uint8> Bytes((const uint8*)Utf8.Get(), Utf8.Length());
+	TUniquePtr<FHttpServerResponse> Response = FHttpServerResponse::Create(MoveTemp(Bytes), TEXT("application/json"));
+	Response->Code = Code;
+	return Response;
+}
+
+// ---------------------------------------------------------------------------
+// Helper: get a named path param as int32 (returns 0 on failure)
+// ---------------------------------------------------------------------------
+static int32 PathParamInt(const FHttpServerRequest& Request, const FString& Name)
+{
+	const FString* Val = Request.PathParams.Find(Name);
+	if (!Val) return 0;
+	return FCString::Atoi(**Val);
+}
+
+// ---------------------------------------------------------------------------
+// Log helper: one line per request handled
+// ---------------------------------------------------------------------------
+static FString VerbToString(EHttpServerRequestVerbs Verb)
+{
+	switch (Verb)
+	{
+	case EHttpServerRequestVerbs::VERB_GET:     return TEXT("GET");
+	case EHttpServerRequestVerbs::VERB_POST:    return TEXT("POST");
+	case EHttpServerRequestVerbs::VERB_PUT:     return TEXT("PUT");
+	case EHttpServerRequestVerbs::VERB_PATCH:   return TEXT("PATCH");
+	case EHttpServerRequestVerbs::VERB_DELETE:  return TEXT("DELETE");
+	case EHttpServerRequestVerbs::VERB_OPTIONS: return TEXT("OPTIONS");
+	case EHttpServerRequestVerbs::VERB_HEAD:    return TEXT("HEAD");
+	case EHttpServerRequestVerbs::VERB_TRACE:   return TEXT("TRACE");
+	case EHttpServerRequestVerbs::VERB_CONNECT: return TEXT("CONNECT");
+	default:                                    return TEXT("UNKNOWN");
+	}
+}
+
+static void LogRequest(const FHttpServerRequest& Request, int32 ResponseCode)
+{
+	UE_LOG(LogHttpServer, Display, TEXT("%s %s => %d"),
+		*VerbToString(Request.Verb),
+		*Request.FullPath.GetPath(),
+		ResponseCode);
+}
+
+// ---------------------------------------------------------------------------
+// Handlers
+// ---------------------------------------------------------------------------
+
+// GET/POST/PUT/PATCH/DELETE/OPTIONS /webtests/httptests/methods
+static bool HandleMethods(const FHttpServerRequest& Request, const FHttpResultCallback& OnComplete)
+{
+	LogRequest(Request, 200);
+	OnComplete(JsonResponse(TEXT("{}")));
+	return true;
+}
+
+// GET /webtests/httptests/query_with_params/?var_int=X&var_str=Y
+static bool HandleQueryWithParams(const FHttpServerRequest& Request, const FHttpResultCallback& OnComplete)
+{
+	const FString* VarIntStr = Request.QueryParams.Find(TEXT("var_int"));
+	const FString* VarStr    = Request.QueryParams.Find(TEXT("var_str"));
+
+	int32 VarInt = VarIntStr ? FCString::Atoi(**VarIntStr) : 0;
+	FString StrVal = VarStr ? *VarStr : TEXT("");
+
+	FString Body = FString::Printf(TEXT("{\"var_str\":\"%s\",\"var_int\":%d}"), *StrVal, VarInt);
+	LogRequest(Request, 200);
+	OnComplete(JsonResponse(Body));
+	return true;
+}
+
+// GET /webtests/httptests/get_data_without_chunks/:bytes_number/:repeat_at/
+static bool HandleGetDataWithoutChunks(const FHttpServerRequest& Request, const FHttpResultCallback& OnComplete)
+{
+	int32 BytesNumber = PathParamInt(Request, TEXT("bytes_number"));
+	int32 RepeatAt    = PathParamInt(Request, TEXT("repeat_at"));
+
+	if (BytesNumber <= 0 || RepeatAt <= 0 || RepeatAt > 10)
+	{
+		OnComplete(FHttpServerResponse::Error(EHttpServerResponseCodes::BadRequest));
+		return true;
+	}
+
+	TArray<uint8> Data;
+	Data.Reserve(BytesNumber);
+	for (int32 i = 0; i < BytesNumber; ++i)
+	{
+		Data.Add(static_cast<uint8>('0' + (i % RepeatAt)));
+	}
+
+	LogRequest(Request, 200);
+	TUniquePtr<FHttpServerResponse> Response = FHttpServerResponse::Create(MoveTemp(Data), TEXT("application/octet-stream"));
+	Response->Code = EHttpServerResponseCodes::Ok;
+	OnComplete(MoveTemp(Response));
+	return true;
+}
+
+// GET /webtests/httptests/get_data_gzip/:bytes_number/:repeat_at/
+static bool HandleGetDataGzip(const FHttpServerRequest& Request, const FHttpResultCallback& OnComplete)
+{
+	int32 BytesNumber = PathParamInt(Request, TEXT("bytes_number"));
+	int32 RepeatAt    = PathParamInt(Request, TEXT("repeat_at"));
+
+	if (BytesNumber <= 0 || RepeatAt <= 0 || RepeatAt > 10)
+	{
+		OnComplete(FHttpServerResponse::Error(EHttpServerResponseCodes::BadRequest));
+		return true;
+	}
+
+	TArray<uint8> RawData;
+	RawData.Reserve(BytesNumber);
+	for (int32 i = 0; i < BytesNumber; ++i)
+	{
+		RawData.Add(static_cast<uint8>('0' + (i % RepeatAt)));
+	}
+
+	int64 CompressBound = 0;
+	FCompression::CompressMemoryBound(NAME_Gzip, CompressBound, (int64)RawData.Num());
+	TArray<uint8> Compressed;
+	Compressed.SetNumUninitialized((int32)CompressBound);
+	int64 CompressedSize = CompressBound;
+	FCompression::CompressMemory(NAME_Gzip, Compressed.GetData(), CompressedSize, RawData.GetData(), (int64)RawData.Num());
+	Compressed.SetNum((int32)CompressedSize);
+
+	LogRequest(Request, 200);
+	TUniquePtr<FHttpServerResponse> Response = FHttpServerResponse::Create(MoveTemp(Compressed), TEXT("application/octet-stream"));
+	Response->Code = EHttpServerResponseCodes::Ok;
+	Response->Headers.Add(TEXT("Content-Encoding"), { TEXT("gzip") });
+	OnComplete(MoveTemp(Response));
+	return true;
+}
+
+// ---------------------------------------------------------------------------
+// Background runnable: streams chunks one at a time via StreamingBodyQueue,
+// with per-chunk latency sleep, so the client receives data incrementally.
+// ---------------------------------------------------------------------------
+class FStreamingDownloadRunnable : public FRunnable
+{
+public:
+	FStreamingDownloadRunnable(int32 InChunks, int32 InChunkSize, int32 InLatencySecs,
+		TSharedPtr<TQueue<TArray<uint8>, EQueueMode::Spsc>> InQueue,
+		TSharedPtr<TAtomic<bool>> InComplete)
+		: Chunks(InChunks), ChunkSize(InChunkSize), LatencySecs(InLatencySecs)
+		, Queue(MoveTemp(InQueue)), Complete(MoveTemp(InComplete))
+	{}
+
+	virtual uint32 Run() override
+	{
+		UE_LOG(LogHttpServer, Display, TEXT("streaming_download: %d chunks x %d bytes, latency %ds each"), Chunks, ChunkSize, LatencySecs);
+
+		TArray<uint8> ChunkData;
+		ChunkData.SetNumUninitialized(ChunkSize);
+		FMemory::Memset(ChunkData.GetData(), (uint8)'d', ChunkSize);
+
+		for (int32 i = 0; i < Chunks; ++i)
+		{
+			if (LatencySecs > 0)
+			{
+				FPlatformProcess::SleepNoStats((float)LatencySecs);
+			}
+			Queue->Enqueue(ChunkData);
+		}
+
+		// Signal that all chunks have been enqueued
+		Complete->Store(true, EMemoryOrder::Relaxed);
+		return 0;
+	}
+
+	virtual void Exit() override { delete this; }
+
+private:
+	int32 Chunks;
+	int32 ChunkSize;
+	int32 LatencySecs;
+	TSharedPtr<TQueue<TArray<uint8>, EQueueMode::Spsc>> Queue;
+	TSharedPtr<TAtomic<bool>> Complete;
+};
+
+// GET /webtests/httptests/streaming_download/:chunks/:chunk_size/:chunk_latency/
+static bool HandleStreamingDownload(const FHttpServerRequest& Request, const FHttpResultCallback& OnComplete)
+{
+	int32 Chunks      = PathParamInt(Request, TEXT("chunks"));
+	int32 ChunkSize   = PathParamInt(Request, TEXT("chunk_size"));
+	int32 LatencySecs = PathParamInt(Request, TEXT("chunk_latency"));
+
+	if (Chunks <= 0 || ChunkSize <= 0)
+	{
+		OnComplete(FHttpServerResponse::Error(EHttpServerResponseCodes::BadRequest));
+		return true;
+	}
+
+	LogRequest(Request, 200);
+
+	int64 TotalSize = (int64)Chunks * ChunkSize;
+	auto Queue    = MakeShared<TQueue<TArray<uint8>, EQueueMode::Spsc>>();
+	auto Complete = MakeShared<TAtomic<bool>>(false);
+
+	TUniquePtr<FHttpServerResponse> ResponseUnique = MakeUnique<FHttpServerResponse>();
+	ResponseUnique->Code = EHttpServerResponseCodes::Ok;
+	ResponseUnique->Headers.Add(TEXT("content-type"), { TEXT("application/octet-stream") });
+	ResponseUnique->Headers.Add(TEXT("content-length"), { FString::FromInt((int32)TotalSize) });
+	ResponseUnique->StreamingBodyQueue   = Queue;
+	ResponseUnique->StreamingBodyComplete = Complete;
+
+	OnComplete(MoveTemp(ResponseUnique));
+
+	// FRunnable::Exit() deletes itself after thread finishes
+	FStreamingDownloadRunnable* Runnable = new FStreamingDownloadRunnable(Chunks, ChunkSize, LatencySecs, Queue, Complete);
+	FRunnableThread::Create(Runnable, TEXT("StreamingDownload"), 0, TPri_Normal);
+	return true;
+}
+
+// POST /webtests/httptests/streaming_upload_post
+static bool HandleStreamingUploadPost(const FHttpServerRequest& Request, const FHttpResultCallback& OnComplete)
+{
+	LogRequest(Request, 200);
+	UE_LOG(LogHttpServer, Display, TEXT("  body: %d bytes"), Request.Body.Num());
+	OnComplete(JsonResponse(TEXT("{}")));
+	return true;
+}
+
+// PUT /webtests/httptests/streaming_upload_put
+static bool HandleStreamingUploadPut(const FHttpServerRequest& Request, const FHttpResultCallback& OnComplete)
+{
+	LogRequest(Request, 200);
+	UE_LOG(LogHttpServer, Display, TEXT("  body: %d bytes"), Request.Body.Num());
+	OnComplete(JsonResponse(TEXT("{}")));
+	return true;
+}
+
+// GET /webtests/httptests/redirect_from
+static bool HandleRedirectFrom(const FHttpServerRequest& Request, const FHttpResultCallback& OnComplete)
+{
+	LogRequest(Request, 302);
+	TUniquePtr<FHttpServerResponse> Response = MakeUnique<FHttpServerResponse>();
+	Response->Code = EHttpServerResponseCodes::Redirect;
+	Response->Headers.Add(TEXT("Location"), { TEXT("/webtests/httptests/redirect_to") });
+	OnComplete(MoveTemp(Response));
+	return true;
+}
+
+// GET /webtests/httptests/redirect_to
+static bool HandleRedirectTo(const FHttpServerRequest& Request, const FHttpResultCallback& OnComplete)
+{
+	LogRequest(Request, 200);
+	OnComplete(JsonResponse(TEXT("{}")));
+	return true;
+}
+
+// ---------------------------------------------------------------------------
+// Background runnable: sleeps N seconds, then enqueues response to game thread
+// ---------------------------------------------------------------------------
+class FMockLatencyRunnable : public FRunnable
+{
+public:
+	FMockLatencyRunnable(int32 InSeconds, FHttpResultCallback InOnComplete)
+		: Seconds(InSeconds), OnComplete(MoveTemp(InOnComplete))
+	{}
+
+	virtual uint32 Run() override
+	{
+		UE_LOG(LogHttpServer, Display, TEXT("mock_latency: sleeping %ds"), Seconds);
+		FPlatformProcess::SleepNoStats((float)Seconds);
+		EnqueueDeferredResponse(MoveTemp(OnComplete), JsonResponse(TEXT("{}")));
+		return 0;
+	}
+
+	virtual void Exit() override { delete this; }
+
+private:
+	int32 Seconds;
+	FHttpResultCallback OnComplete;
+};
+
+// GET /webtests/httptests/mock_latency/:latency/
+static bool HandleMockLatency(const FHttpServerRequest& Request, const FHttpResultCallback& OnComplete)
+{
+	int32 Latency = PathParamInt(Request, TEXT("latency"));
+	LogRequest(Request, 200);
+	FMockLatencyRunnable* Runnable = new FMockLatencyRunnable(Latency, OnComplete);
+	FRunnableThread::Create(Runnable, TEXT("MockLatency"), 0, TPri_Normal);
+	return true;
+}
+
+// GET,POST,PUT,DELETE /webtests/httptests/mock_status/:status_code/
+static bool HandleMockStatus(const FHttpServerRequest& Request, const FHttpResultCallback& OnComplete)
+{
+	int32 StatusCode = PathParamInt(Request, TEXT("status_code"));
+	TUniquePtr<FHttpServerResponse> Response = JsonResponse(
+		TEXT("{\"key_a\":\"value_a\",\"key_b\":\"value_b\"}"),
+		static_cast<EHttpServerResponseCodes>(StatusCode));
+
+	const TArray<FString>* RetryAfterValues = Request.Headers.Find(TEXT("Retry-After"));
+	if (RetryAfterValues && RetryAfterValues->Num() > 0)
+	{
+		Response->Headers.Add(TEXT("Retry-After"), *RetryAfterValues);
+	}
+
+	LogRequest(Request, StatusCode);
+	OnComplete(MoveTemp(Response));
+	return true;
+}
+
+// POST /webtests/httptests/post_echo_body
+static bool HandlePostEchoBody(const FHttpServerRequest& Request, const FHttpResultCallback& OnComplete)
+{
+	LogRequest(Request, 200);
+	TArray<uint8> Body = Request.Body;
+	TUniquePtr<FHttpServerResponse> Response = FHttpServerResponse::Create(MoveTemp(Body), TEXT("application/json"));
+	Response->Code = EHttpServerResponseCodes::Ok;
+	OnComplete(MoveTemp(Response));
+	return true;
+}
+
+// ---------------------------------------------------------------------------
+// Public entry point
+// ---------------------------------------------------------------------------
+void StartHttpServer(uint32 Port)
+{
+	TSharedPtr<IHttpRouter> Router = FHttpServerModule::Get().GetHttpRouter(Port, /*bFailOnBindFailure=*/true);
+	if (!Router.IsValid())
+	{
+		UE_LOG(LogHttpServer, Error, TEXT("Failed to bind HTTP router on port %d"), Port);
+		return;
+	}
+
+	using EVerb = EHttpServerRequestVerbs;
+	const EVerb AllMethods = EVerb::VERB_GET | EVerb::VERB_POST | EVerb::VERB_PUT | EVerb::VERB_PATCH
+	                       | EVerb::VERB_DELETE | EVerb::VERB_OPTIONS | EVerb::VERB_HEAD | EVerb::VERB_TRACE | EVerb::VERB_CONNECT;
+
+	Router->BindRoute(FHttpPath(TEXT("/webtests/httptests/methods")),
+		AllMethods,
+		FHttpRequestHandler::CreateStatic(&HandleMethods));
+	Router->BindRoute(FHttpPath(TEXT("/webtests/httptests/query_with_params/")),
+		EVerb::VERB_GET,
+		FHttpRequestHandler::CreateStatic(&HandleQueryWithParams));
+	Router->BindRoute(FHttpPath(TEXT("/webtests/httptests/get_data_without_chunks/:bytes_number/:repeat_at/")),
+		EVerb::VERB_GET,
+		FHttpRequestHandler::CreateStatic(&HandleGetDataWithoutChunks));
+	Router->BindRoute(FHttpPath(TEXT("/webtests/httptests/get_data_gzip/:bytes_number/:repeat_at/")),
+		EVerb::VERB_GET,
+		FHttpRequestHandler::CreateStatic(&HandleGetDataGzip));
+	Router->BindRoute(FHttpPath(TEXT("/webtests/httptests/streaming_download/:chunks/:chunk_size/:chunk_latency/")),
+		EVerb::VERB_GET,
+		FHttpRequestHandler::CreateStatic(&HandleStreamingDownload));
+	Router->BindRoute(FHttpPath(TEXT("/webtests/httptests/streaming_upload_post")),
+		EVerb::VERB_POST,
+		FHttpRequestHandler::CreateStatic(&HandleStreamingUploadPost));
+	Router->BindRoute(FHttpPath(TEXT("/webtests/httptests/streaming_upload_put")),
+		EVerb::VERB_PUT,
+		FHttpRequestHandler::CreateStatic(&HandleStreamingUploadPut));
+	Router->BindRoute(FHttpPath(TEXT("/webtests/httptests/redirect_from")),
+		EVerb::VERB_GET,
+		FHttpRequestHandler::CreateStatic(&HandleRedirectFrom));
+	Router->BindRoute(FHttpPath(TEXT("/webtests/httptests/redirect_to")),
+		EVerb::VERB_GET,
+		FHttpRequestHandler::CreateStatic(&HandleRedirectTo));
+	Router->BindRoute(FHttpPath(TEXT("/webtests/httptests/mock_latency/:latency/")),
+		EVerb::VERB_GET,
+		FHttpRequestHandler::CreateStatic(&HandleMockLatency));
+	Router->BindRoute(FHttpPath(TEXT("/webtests/httptests/mock_status/:status_code/")),
+		EVerb::VERB_GET | EVerb::VERB_POST | EVerb::VERB_PUT | EVerb::VERB_DELETE,
+		FHttpRequestHandler::CreateStatic(&HandleMockStatus));
+	Router->BindRoute(FHttpPath(TEXT("/webtests/httptests/post_echo_body")),
+		EVerb::VERB_POST,
+		FHttpRequestHandler::CreateStatic(&HandlePostEchoBody));
+
+	FHttpServerModule::Get().StartAllListeners();
+	UE_LOG(LogHttpServer, Display, TEXT("HTTP server listening on http://0.0.0.0:%d/"), Port);
+}
diff --git a/Engine/Source/Programs/WebTestsServerCpp/Private/WebTestsServer.cpp b/Engine/Source/Programs/WebTestsServerCpp/Private/WebTestsServer.cpp
new file mode 100644
--- /dev/null
+++ b/Engine/Source/Programs/WebTestsServerCpp/Private/WebTestsServer.cpp
@@ -0,0 +1,75 @@
+// Copyright Epic Games, Inc. All Rights Reserved.
+
+#include "WebTestsServer.h"
+#include "Containers/Ticker.h"
+#include "HAL/CriticalSection.h"
+#include "HAL/PlatformProcess.h"
+#include "Misc/CoreDelegates.h"
+#include "Misc/ScopeLock.h"
+#include "HttpServerModule.h"
+#include "WebSocketServerModule.h"
+#include "RequiredProgramMainCPPInclude.h"
+
+IMPLEMENT_APPLICATION(WebTestsServerCpp, "WebTestsServerCpp");
+
+// ---------------------------------------------------------------------------
+// Deferred response queue implementation (shared with WebTestsHttpServer.cpp)
+// ---------------------------------------------------------------------------
+static FCriticalSection GDeferredResponsesLock;
+static TArray<FDeferredResponse> GDeferredResponses;
+
+void EnqueueDeferredResponse(FHttpResultCallback&& OnComplete, TUniquePtr<FHttpServerResponse>&& Response)
+{
+	FScopeLock Lock(&GDeferredResponsesLock);
+	GDeferredResponses.Add({ MoveTemp(OnComplete), MoveTemp(Response) });
+}
+
+void FlushDeferredResponses()
+{
+	TArray<FDeferredResponse> ToFlush;
+	{
+		FScopeLock Lock(&GDeferredResponsesLock);
+		ToFlush = MoveTemp(GDeferredResponses);
+	}
+	for (FDeferredResponse& Deferred : ToFlush)
+	{
+		Deferred.OnComplete(MoveTemp(Deferred.Response));
+	}
+}
+
+// ---------------------------------------------------------------------------
+// Entry point
+// ---------------------------------------------------------------------------
+INT32_MAIN_INT32_ARGC_TCHAR_ARGV()
+{
+	FTaskTagScope Scope(ETaskTag::EGameThread);
+	ON_SCOPE_EXIT
+	{
+		LLM(FLowLevelMemTracker::Get().UpdateStatsPerFrame());
+		RequestEngineExit(TEXT("Exiting"));
+		FEngineLoop::AppPreExit();
+		FModuleManager::Get().UnloadModulesAtShutdown();
+		FEngineLoop::AppExit();
+	};
+
+	if (int32 Ret = GEngineLoop.PreInit(ArgC, ArgV))
+	{
+		return Ret;
+	}
+
+	StartHttpServer(8000);
+	StartWebSocketServer(8001);
+
+	const float TickInterval = 1.0f / 60.0f;
+	while (!IsEngineExitRequested())
+	{
+		FlushDeferredResponses();
+		FTSTicker::GetCoreTicker().Tick(TickInterval);
+		FPlatformProcess::SleepNoStats(TickInterval);
+	}
+
+	FHttpServerModule::Get().StopAllListeners();
+	FWebSocketServerModule::Get().StopAllServers();
+
+	return 0;
+}
diff --git a/Engine/Source/Programs/WebTestsServerCpp/Private/WebTestsServer.h b/Engine/Source/Programs/WebTestsServerCpp/Private/WebTestsServer.h
new file mode 100644
--- /dev/null
+++ b/Engine/Source/Programs/WebTestsServerCpp/Private/WebTestsServer.h
@@ -0,0 +1,22 @@
+// Copyright Epic Games, Inc. All Rights Reserved.
+#pragma once
+
+#include "CoreMinimal.h"
+#include "HttpResultCallback.h"
+#include "HttpServerResponse.h"
+
+// ---------------------------------------------------------------------------
+// Deferred response queue — background threads enqueue (callback, response)
+// pairs; the main tick loop drains them on the game thread.
+// ---------------------------------------------------------------------------
+struct FDeferredResponse
+{
+	FHttpResultCallback             OnComplete;
+	TUniquePtr<FHttpServerResponse> Response;
+};
+
+void EnqueueDeferredResponse(FHttpResultCallback&& OnComplete, TUniquePtr<FHttpServerResponse>&& Response);
+void FlushDeferredResponses();
+
+void StartHttpServer(uint32 Port);
+void StartWebSocketServer(uint32 Port);
diff --git a/Engine/Source/Programs/WebTestsServerCpp/Private/WebTestsWebSocketServer.cpp b/Engine/Source/Programs/WebTestsServerCpp/Private/WebTestsWebSocketServer.cpp
new file mode 100644
--- /dev/null
+++ b/Engine/Source/Programs/WebTestsServerCpp/Private/WebTestsWebSocketServer.cpp
@@ -0,0 +1,40 @@
+// Copyright Epic Games, Inc. All Rights Reserved.
+
+#include "CoreMinimal.h"
+#include "IWebSocketClientConnection.h"
+#include "IWebSocketServer.h"
+#include "WebSocketServerModule.h"
+
+void StartWebSocketServer(uint32 Port)
+{
+	TSharedPtr<IWebSocketServer> WsServer = FWebSocketServerModule::Get().GetWebSocketServer(Port);
+
+	WsServer->OnConnected([](TSharedRef<IWebSocketClientConnection> Conn)
+	{
+		UE_LOG(LogWebSocketServer, Display, TEXT("WS  CONNECT  %s"), *Conn->GetPath());
+	});
+
+	WsServer->OnMessage([](TSharedRef<IWebSocketClientConnection> Conn, const FString& Msg)
+	{
+		const FString Path = Conn->GetPath();
+		UE_LOG(LogWebSocketServer, Display, TEXT("WS  MESSAGE  path='%s' msg='%s'"), *Path, *Msg);
+		if (Path.EndsWith(TEXT("/echo/")))
+		{
+			UE_LOG(LogWebSocketServer, Display, TEXT("WS  ECHO     sending back '%s'"), *Msg);
+			Conn->SendText(Msg);
+		}
+		else if (Path.EndsWith(TEXT("/close_on_receive_message/")))
+		{
+			UE_LOG(LogWebSocketServer, Display, TEXT("WS  CLOSE    closing connection"));
+			Conn->Close(1000);
+		}
+	});
+
+	WsServer->OnDisconnected([](TSharedRef<IWebSocketClientConnection> Conn)
+	{
+		UE_LOG(LogWebSocketServer, Display, TEXT("WS  DISCONNECT %s"), *Conn->GetPath());
+	});
+
+	FWebSocketServerModule::Get().StartAllServers();
+	UE_LOG(LogWebSocketServer, Display, TEXT("WebSocket server listening on ws://0.0.0.0:%d/"), Port);
+}
diff --git a/Engine/Source/Programs/WebTestsServerCpp/WebTestsServerCpp.Build.cs b/Engine/Source/Programs/WebTestsServerCpp/WebTestsServerCpp.Build.cs
new file mode 100644
--- /dev/null
+++ b/Engine/Source/Programs/WebTestsServerCpp/WebTestsServerCpp.Build.cs
@@ -0,0 +1,22 @@
+// Copyright Epic Games, Inc. All Rights Reserved.
+
+using UnrealBuildTool;
+
+public class WebTestsServerCpp : ModuleRules
+{
+	public WebTestsServerCpp(ReadOnlyTargetRules Target) : base(Target)
+	{
+		PublicIncludePathModuleNames.Add("Launch");
+
+		PrivateDependencyModuleNames.AddRange(
+			new string[] {
+				"Core",
+				"ApplicationCore",
+				"Projects",
+				"HTTPServer",
+				"WebSocketServer",
+				"Sockets",
+				"Json",
+			});
+	}
+}
diff --git a/Engine/Source/Programs/WebTestsServerCpp/WebTestsServerCpp.Target.cs b/Engine/Source/Programs/WebTestsServerCpp/WebTestsServerCpp.Target.cs
new file mode 100644
--- /dev/null
+++ b/Engine/Source/Programs/WebTestsServerCpp/WebTestsServerCpp.Target.cs
@@ -0,0 +1,22 @@
+// Copyright Epic Games, Inc. All Rights Reserved.
+
+using UnrealBuildTool;
+
+[SupportedPlatforms(UnrealPlatformClass.Desktop)]
+public class WebTestsServerCppTarget : TargetRules
+{
+	public WebTestsServerCppTarget(TargetInfo Target) : base(Target)
+	{
+		Type = TargetType.Program;
+		IncludeOrderVersion = EngineIncludeOrderVersion.Latest;
+		LinkType = TargetLinkType.Monolithic;
+		LaunchModuleName = "WebTestsServerCpp";
+
+		bBuildDeveloperTools = false;
+		bCompileAgainstEngine = false;
+		bCompileAgainstCoreUObject = false;
+		bCompileAgainstApplicationCore = true;
+		bCompileICU = false;
+		bIsBuildingConsoleApplication = true;
+	}
+}
diff --git a/Engine/Source/Programs/WebTestsServerCpp/runserver.bat b/Engine/Source/Programs/WebTestsServerCpp/runserver.bat
new file mode 100755
--- /dev/null
+++ b/Engine/Source/Programs/WebTestsServerCpp/runserver.bat
@@ -0,0 +1,15 @@
+@echo off
+setlocal
+
+set SCRIPT_DIR=%~dp0
+set UE_ROOT=%SCRIPT_DIR%..\..\..\..
+
+echo Building WebTestsServerCpp...
+call "%UE_ROOT%\Engine\Build\BatchFiles\Build.bat" WebTestsServerCpp Win64 Development
+if %ERRORLEVEL% neq 0 (
+    echo Build failed.
+    exit /b %ERRORLEVEL%
+)
+
+echo Starting WebTestsServerCpp...
+"%UE_ROOT%\Engine\Binaries\Win64\WebTestsServerCpp.exe" -log -NOLOGTIMES
diff --git a/Engine/Source/Programs/WebTestsServerCpp/runserver.sh b/Engine/Source/Programs/WebTestsServerCpp/runserver.sh
new file mode 100755
--- /dev/null
+++ b/Engine/Source/Programs/WebTestsServerCpp/runserver.sh
@@ -0,0 +1,20 @@
+#!/bin/bash
+
+SCRIPT_DIR="$(cd "$(dirname "${BASH_SOURCE[0]}")" && pwd)"
+UE_ROOT="$SCRIPT_DIR/../../../.."
+
+if [ "$(uname)" = "Darwin" ]; then
+    PLATFORM="Mac"
+else
+    PLATFORM="Linux"
+fi
+
+echo "Building WebTestsServerCpp..."
+"$UE_ROOT/Engine/Build/BatchFiles/RunUBT.sh" WebTestsServerCpp "$PLATFORM" Development
+if [ $? -ne 0 ]; then
+    echo "Build failed."
+    exit 1
+fi
+
+echo "Starting WebTestsServerCpp..."
+"$UE_ROOT/Engine/Binaries/$PLATFORM/WebTestsServerCpp" -log -NOLOGTIMES
diff --git a/Engine/Source/Runtime/Online/HTTPServer/Private/HttpConnection.cpp b/Engine/Source/Runtime/Online/HTTPServer/Private/HttpConnection.cpp
--- a/Engine/Source/Runtime/Online/HTTPServer/Private/HttpConnection.cpp
+++ b/Engine/Source/Runtime/Online/HTTPServer/Private/HttpConnection.cpp
@@ -164,19 +164,20 @@ void FHttpConnection::CompleteRead(const TSharedPtr<FHttpServerRequest>& Request
 	}
 
 	TWeakPtr<FHttpConnection> WeakThisPtr(AsShared());
-	FHttpResultCallback OnProcessingComplete = 
-		[WeakThisPtr, 
+	FHttpResultCallback OnProcessingComplete =
+		[WeakThisPtr,
 		OriginPortCapture = OriginPort,
-		ConnectionIdCapture = ConnectionId, 
+		ConnectionIdCapture = ConnectionId,
 		LastRequestNumberCapture = LastRequestNumber,
-		ResponseVersionCapture = Request->HttpVersion]
+		ResponseVersionCapture = Request->HttpVersion,
+		bIsHeadRequest = (Request->Verb == EHttpServerRequestVerbs::VERB_HEAD)]
 	(TUniquePtr<FHttpServerResponse>&& Response)
 	{
 		TSharedPtr<FHttpConnection> SharedThisPtr = WeakThisPtr.Pin();
 		if (SharedThisPtr.IsValid())
 		{
 			UE_LOG(LogHttpConnection, Verbose,
-				TEXT("EndProcessRequest\t [%d][%u]-%u : %u"), 
+				TEXT("EndProcessRequest\t [%d][%u]-%u : %u"),
 						 OriginPortCapture, ConnectionIdCapture, LastRequestNumberCapture, int(Response->Code));
 
 			// Ensure this result callback was called at least once
@@ -184,6 +185,11 @@ void FHttpConnection::CompleteRead(const TSharedPtr<FHttpServerRequest>& Request
 
 			// Begin response flow
 			Response->HttpVersion = ResponseVersionCapture;
+			// HEAD responses must have the same headers as GET but no body
+			if (bIsHeadRequest)
+			{
+				Response->Body.Empty();
+			}
 			SharedThisPtr->BeginWrite(MoveTemp(Response), LastRequestNumberCapture);
 		}
 	};
diff --git a/Engine/Source/Runtime/Online/HTTPServer/Private/HttpConnectionRequestReadContext.cpp b/Engine/Source/Runtime/Online/HTTPServer/Private/HttpConnectionRequestReadContext.cpp
--- a/Engine/Source/Runtime/Online/HTTPServer/Private/HttpConnectionRequestReadContext.cpp
+++ b/Engine/Source/Runtime/Online/HTTPServer/Private/HttpConnectionRequestReadContext.cpp
@@ -273,6 +273,21 @@ TSharedPtr<FHttpServerRequest> FHttpConnectionRequestReadContext::BuildRequest(c
 	{
 		Request->Verb = EHttpServerRequestVerbs::VERB_OPTIONS;
 	}
+	else if (0 == RequestVerb.Compare(TEXT("HEAD"),
+		ESearchCase::IgnoreCase))
+	{
+		Request->Verb = EHttpServerRequestVerbs::VERB_HEAD;
+	}
+	else if (0 == RequestVerb.Compare(TEXT("TRACE"),
+		ESearchCase::IgnoreCase))
+	{
+		Request->Verb = EHttpServerRequestVerbs::VERB_TRACE;
+	}
+	else if (0 == RequestVerb.Compare(TEXT("CONNECT"),
+		ESearchCase::IgnoreCase))
+	{
+		Request->Verb = EHttpServerRequestVerbs::VERB_CONNECT;
+	}
 	else
 	{
 		// Unknown Verb
@@ -303,6 +318,7 @@ TSharedPtr<FHttpServerRequest> FHttpConnectionRequestReadContext::BuildRequest(c
 		}
 	}
 	Request->RelativePath.SetPath(RequestHttpPath);
+	Request->FullPath.SetPath(RequestHttpPath);
 
 	// Parse/store http version
 	const FString& RequestHttpVersion = HttpMethodTokens[2];
diff --git a/Engine/Source/Runtime/Online/HTTPServer/Private/HttpConnectionResponseWriteContext.cpp b/Engine/Source/Runtime/Online/HTTPServer/Private/HttpConnectionResponseWriteContext.cpp
--- a/Engine/Source/Runtime/Online/HTTPServer/Private/HttpConnectionResponseWriteContext.cpp
+++ b/Engine/Source/Runtime/Online/HTTPServer/Private/HttpConnectionResponseWriteContext.cpp
@@ -28,7 +28,7 @@ void FHttpConnectionResponseWriteContext::ResetContext(TUniquePtr<FHttpServerRes
 	{
 		const bool bResponseIsStream = EnumHasAnyFlags(Response->Flags, EHttpServerResponseFlags::MultipleWriteStream);
 
-		if (!bResponseIsStream)
+		if (!bResponseIsStream || !Response->StreamingBodyQueue.IsValid())
 		{
 			// Affix Content-Length Header
 			TArray<FString> ContentLengthValue = { FString::FromInt(Response->Body.Num()) };
@@ -63,6 +63,16 @@ EHttpConnectionContextState FHttpConnectionResponseWriteContext::WriteStream(flo
 		HeaderBytesWritten += BytesWritten;
 	}
 
+	// If streaming queue is set, pull available chunks into Body before writing
+	if (IsWriteHeaderComplete() && Response->StreamingBodyQueue.IsValid())
+	{
+		TArray<uint8> Chunk;
+		while (Response->StreamingBodyQueue->Dequeue(Chunk))
+		{
+			Response->Body.Append(Chunk);
+		}
+	}
+
 	if (IsWriteHeaderComplete() && !IsWriteBodyComplete())
 	{
 		const uint8* DataOffset = Response->Body.GetData() + BodyBytesWritten;
@@ -77,6 +87,12 @@ EHttpConnectionContextState FHttpConnectionResponseWriteContext::WriteStream(flo
 
 	if (IsWriteHeaderComplete() && IsWriteBodyComplete())
 	{
+		// If streaming is in progress, stay in Continue until all chunks are enqueued and written
+		if (Response->StreamingBodyQueue.IsValid() &&
+			(!Response->StreamingBodyComplete.IsValid() || !Response->StreamingBodyComplete->Load(EMemoryOrder::Relaxed)))
+		{
+			return EHttpConnectionContextState::Continue;
+		}
 		ensure(BodyBytesWritten+HeaderBytesWritten == HeaderBytes.Num()+Response->Body.Num());
 		return EHttpConnectionContextState::Done;
 	}
diff --git a/Engine/Source/Runtime/Online/HTTPServer/Public/HttpServerRequest.h b/Engine/Source/Runtime/Online/HTTPServer/Public/HttpServerRequest.h
--- a/Engine/Source/Runtime/Online/HTTPServer/Public/HttpServerRequest.h
+++ b/Engine/Source/Runtime/Online/HTTPServer/Public/HttpServerRequest.h
@@ -15,7 +15,10 @@ enum class EHttpServerRequestVerbs : uint16
 	VERB_PUT    = 1 << 2,
 	VERB_PATCH  = 1 << 3,
 	VERB_DELETE = 1 << 4,
-	VERB_OPTIONS = 1 << 5
+	VERB_OPTIONS = 1 << 5,
+	VERB_HEAD    = 1 << 6,
+	VERB_TRACE   = 1 << 7,
+	VERB_CONNECT = 1 << 8,
 };
 
 ENUM_CLASS_FLAGS(EHttpServerRequestVerbs);
@@ -30,6 +33,9 @@ public:
 	/** The IP address of the peer that initiated the request. */
 	TSharedPtr<FInternetAddr> PeerAddress;
 
+	/** The full HTTP path as originally received */
+	FHttpPath FullPath;
+
 	/** The handler-route-relative HTTP path */
 	FHttpPath RelativePath;
 
diff --git a/Engine/Source/Runtime/Online/HTTPServer/Public/HttpServerResponse.h b/Engine/Source/Runtime/Online/HTTPServer/Public/HttpServerResponse.h
--- a/Engine/Source/Runtime/Online/HTTPServer/Public/HttpServerResponse.h
+++ b/Engine/Source/Runtime/Online/HTTPServer/Public/HttpServerResponse.h
@@ -3,11 +3,14 @@
 #include "Containers/Array.h"
 #include "Containers/ArrayView.h"
 #include "Containers/Map.h"
+#include "Containers/Queue.h"
 #include "Containers/UnrealString.h"
 #include "CoreMinimal.h"
 #include "HAL/Platform.h"
 #include "HttpServerConstants.h"
 #include "HttpServerHttpVersion.h"
+#include "Templates/Atomic.h"
+#include "Templates/SharedPointer.h"
 #include "Templates/UniquePtr.h"
 #include "Templates/UnrealTemplate.h"
 
@@ -63,6 +66,15 @@ public:
 	/** Response flags */
 	EHttpServerResponseFlags Flags = EHttpServerResponseFlags::None;
 
+	/**
+	 * Optional: streaming body queue. When set, WriteStream will drain this queue chunk-by-chunk
+	 * in addition to (after) the initial Body content.  Must be an SPSC queue: one producer thread
+	 * enqueues chunks, the game thread (WriteStream) dequeues.
+	 * Set bStreamingBodyComplete to true after all chunks have been enqueued.
+	 */
+	TSharedPtr<TQueue<TArray<uint8>, EQueueMode::Spsc>> StreamingBodyQueue;
+	TSharedPtr<TAtomic<bool>> StreamingBodyComplete;
+
 public:
 
 	/**
diff --git a/Engine/Source/Runtime/Online/WebSocketServer/Private/LwsWebSocketClientConnection.cpp b/Engine/Source/Runtime/Online/WebSocketServer/Private/LwsWebSocketClientConnection.cpp
new file mode 100644
--- /dev/null
+++ b/Engine/Source/Runtime/Online/WebSocketServer/Private/LwsWebSocketClientConnection.cpp
@@ -0,0 +1,89 @@
+// Copyright Epic Games, Inc. All Rights Reserved.
+
+#include "LwsWebSocketClientConnection.h"
+
+#if WITH_WEBSOCKETSERVER
+
+#include "Misc/ScopeLock.h"
+#include "WebSocketServerModule.h"
+
+FLwsWebSocketClientConnection::FLwsWebSocketClientConnection(lws* InWsi, const FString& InPath, FLwsWebSocketServer* InServer)
+	: Wsi(InWsi)
+	, Server(InServer)
+	, Path(InPath)
+{
+}
+
+void FLwsWebSocketClientConnection::SendText(const FString& Message)
+{
+	FScopeLock Lock(&WsiLock);
+	if (!Wsi)
+	{
+		return;
+	}
+
+	// Build lws frame: LWS_PRE bytes of header space + UTF-8 payload
+	FTCHARToUTF8 Utf8(*Message);
+	const int32 PayloadLen = Utf8.Length();
+	TArray<uint8> Frame;
+	Frame.AddZeroed(LWS_PRE);
+	Frame.Append(reinterpret_cast<const uint8*>(Utf8.Get()), PayloadLen);
+
+	SendQueue.Enqueue(MoveTemp(Frame));
+	lws_callback_on_writable(Wsi);
+}
+
+void FLwsWebSocketClientConnection::Close(int32 Code, const FString& Reason)
+{
+	FScopeLock Lock(&WsiLock);
+	if (!Wsi)
+	{
+		return;
+	}
+
+	// Store close params — lws_close_reason() must be called on the lws thread
+	UE_LOG(LogWebSocketServer, Display, TEXT("WS Close(): storing close request code=%d, calling lws_callback_on_writable"), Code);
+	CloseReason = Reason;
+	CloseCode.Store(Code, EMemoryOrder::Relaxed);
+	bCloseRequested.Store(true, EMemoryOrder::SequentiallyConsistent);
+
+	// Schedule SERVER_WRITEABLE on the lws thread so LwsThreadFlushAndCheckClose() is called
+	lws_callback_on_writable(Wsi);
+}
+
+bool FLwsWebSocketClientConnection::LwsThreadFlushAndCheckClose()
+{
+	// Flush outgoing send queue
+	TArray<uint8> Frame;
+	while (SendQueue.Dequeue(Frame))
+	{
+		const int32 PayloadLen = Frame.Num() - LWS_PRE;
+		if (PayloadLen > 0)
+		{
+			UE_LOG(LogWebSocketServer, Display, TEXT("WS lws_write: %d bytes"), PayloadLen);
+			int32 Written = lws_write(Wsi, Frame.GetData() + LWS_PRE, static_cast<size_t>(PayloadLen), LWS_WRITE_TEXT);
+			UE_LOG(LogWebSocketServer, Display, TEXT("WS lws_write returned: %d"), Written);
+		}
+	}
+
+	// If a close was requested, initiate it now on the lws thread
+	if (bCloseRequested.Load(EMemoryOrder::SequentiallyConsistent))
+	{
+		UE_LOG(LogWebSocketServer, Display, TEXT("WS lws_close_reason: code=%d reason='%s'"), CloseCode.Load(EMemoryOrder::Relaxed), *CloseReason);
+		FTCHARToUTF8 ReasonUtf8(*CloseReason);
+		lws_close_reason(Wsi,
+			static_cast<lws_close_status>(CloseCode.Load(EMemoryOrder::Relaxed)),
+			reinterpret_cast<unsigned char*>(const_cast<ANSICHAR*>(ReasonUtf8.Get())),
+			static_cast<size_t>(ReasonUtf8.Length()));
+		return true; // caller must return -1 from callback to initiate close
+	}
+	return false;
+}
+
+void FLwsWebSocketClientConnection::LwsThreadOnClosed()
+{
+	FScopeLock Lock(&WsiLock);
+	Wsi = nullptr;
+}
+
+#endif // WITH_WEBSOCKETSERVER
diff --git a/Engine/Source/Runtime/Online/WebSocketServer/Private/LwsWebSocketClientConnection.h b/Engine/Source/Runtime/Online/WebSocketServer/Private/LwsWebSocketClientConnection.h
new file mode 100644
--- /dev/null
+++ b/Engine/Source/Runtime/Online/WebSocketServer/Private/LwsWebSocketClientConnection.h
@@ -0,0 +1,64 @@
+// Copyright Epic Games, Inc. All Rights Reserved.
+#pragma once
+
+#include "CoreMinimal.h"
+
+#if WITH_WEBSOCKETSERVER
+
+#include "IWebSocketClientConnection.h"
+#include "Containers/Queue.h"
+#include "HAL/CriticalSection.h"
+#include "Misc/ScopeLock.h"
+#include "Templates/Atomic.h"
+
+#if PLATFORM_WINDOWS
+#	include "Windows/AllowWindowsPlatformTypes.h"
+#endif
+THIRD_PARTY_INCLUDES_START
+#include "libwebsockets.h"
+THIRD_PARTY_INCLUDES_END
+#if PLATFORM_WINDOWS
+#	include "Windows/HideWindowsPlatformTypes.h"
+#endif
+
+class FLwsWebSocketServer;
+
+// ─────────────────────────────────────────────────────────────────────────────
+// Per-connection object — created on ESTABLISHED, destroyed after CLOSED
+// ─────────────────────────────────────────────────────────────────────────────
+class FLwsWebSocketClientConnection final
+	: public IWebSocketClientConnection
+	, public TSharedFromThis<FLwsWebSocketClientConnection>
+{
+public:
+	FLwsWebSocketClientConnection(lws* InWsi, const FString& InPath, FLwsWebSocketServer* InServer);
+
+	// IWebSocketClientConnection
+	virtual void    SendText(const FString& Message) override;
+	virtual void    Close(int32 Code = 1000, const FString& Reason = FString()) override;
+	virtual FString GetPath() const override { return Path; }
+
+	// Called on the lws thread from SERVER_WRITEABLE: flush send queue and handle pending close.
+	// Returns true if the connection should be closed (caller must return -1 from callback).
+	bool LwsThreadFlushAndCheckClose();
+
+	// Called on the lws thread when LWS_CALLBACK_CLOSED / WSI_DESTROY fires
+	void LwsThreadOnClosed();
+
+private:
+	FCriticalSection           WsiLock;
+	lws*                       Wsi;          // Nulled on close
+	FLwsWebSocketServer*       Server;       // Non-owning
+	FString                    Path;
+
+	// SPSC: game/any thread pushes, lws thread pops
+	TQueue<TArray<uint8>, EQueueMode::Spsc> SendQueue;
+
+	// Close request — set from game thread, read on lws thread
+	TAtomic<bool>  bCloseRequested { false };
+	TAtomic<int32> CloseCode       { 1000 };
+	// CloseReason is written once before bCloseRequested is set true; safe to read after
+	FString        CloseReason;
+};
+
+#endif // WITH_WEBSOCKETSERVER
diff --git a/Engine/Source/Runtime/Online/WebSocketServer/Private/LwsWebSocketServer.cpp b/Engine/Source/Runtime/Online/WebSocketServer/Private/LwsWebSocketServer.cpp
new file mode 100644
--- /dev/null
+++ b/Engine/Source/Runtime/Online/WebSocketServer/Private/LwsWebSocketServer.cpp
@@ -0,0 +1,359 @@
+// Copyright Epic Games, Inc. All Rights Reserved.
+
+#include "LwsWebSocketServer.h"
+
+#if WITH_WEBSOCKETSERVER
+
+#include "HAL/RunnableThread.h"
+#include "HAL/PlatformProcess.h"
+#include "HAL/PlatformTime.h"
+#include "Misc/ScopeLock.h"
+#include "WebSocketServerModule.h"
+
+// ─────────────────────────────────────────────────────────────────────────────
+// FLwsWebSocketServer
+// ─────────────────────────────────────────────────────────────────────────────
+
+FLwsWebSocketServer::FLwsWebSocketServer(uint32 InPort)
+	: Port(InPort)
+{
+	FMemory::Memzero(LwsProtocol, sizeof(LwsProtocol));
+}
+
+FLwsWebSocketServer::~FLwsWebSocketServer()
+{
+	StopListening();
+}
+
+bool FLwsWebSocketServer::StartListening()
+{
+	if (bListening)
+	{
+		return true;
+	}
+
+	// Register one protocol named "websocket"
+	static const char ProtocolName[] = "websocket";
+	LwsProtocol[0].name                  = ProtocolName;
+	LwsProtocol[0].callback               = &FLwsWebSocketServer::StaticLwsCallback;
+	LwsProtocol[0].per_session_data_size  = sizeof(FLwsWebSocketClientConnection*);
+	LwsProtocol[0].rx_buffer_size         = 65536;
+	// LwsProtocol[1] stays zeroed — required terminator
+
+	struct lws_context_creation_info Info;
+	FMemory::Memzero(&Info, sizeof(Info));
+	Info.port      = static_cast<int>(Port);
+	Info.protocols = LwsProtocol;
+	Info.user      = this;    // retrieve via lws_context_user()
+	Info.options  |= LWS_SERVER_OPTION_PEER_CERT_NOT_REQUIRED
+	              |  LWS_SERVER_OPTION_DISABLE_OS_CA_CERTS
+	              |  LWS_SERVER_OPTION_DO_SSL_GLOBAL_INIT
+	              |  LWS_SERVER_OPTION_VALIDATE_UTF8
+	              |  LWS_SERVER_OPTION_DISABLE_IPV6;  // Bind IPv4 only (0.0.0.0), not IPv6
+
+	// Silence most lws noise; only show errors/warnings
+	lws_set_log_level(LLL_ERR | LLL_WARN, nullptr);
+
+	LwsContext = lws_create_context(&Info);
+	if (!LwsContext)
+	{
+		UE_LOG(LogWebSocketServer, Error, TEXT("FLwsWebSocketServer: Failed to create lws context on port %u"), Port);
+		return false;
+	}
+
+	ExitRequest.Set(0);
+	Thread = FRunnableThread::Create(this, TEXT("WebSocketServerThread"), 128 * 1024, TPri_BelowNormal);
+	if (!Thread)
+	{
+		UE_LOG(LogWebSocketServer, Error, TEXT("FLwsWebSocketServer: Failed to create service thread"));
+		lws_context_destroy(LwsContext);
+		LwsContext = nullptr;
+		return false;
+	}
+
+	bListening = true;
+	return true;
+}
+
+void FLwsWebSocketServer::StopListening()
+{
+	if (!bListening)
+	{
+		return;
+	}
+	bListening = false;
+
+	if (Thread)
+	{
+		ExitRequest.Set(1);
+		WakeService();
+		Thread->Kill(true);
+		delete Thread;
+		Thread = nullptr;
+	}
+
+	if (LwsContext)
+	{
+		lws_context_destroy(LwsContext);
+		LwsContext = nullptr;
+	}
+
+	{
+		FScopeLock Lock(&ActiveConnectionsLock);
+		ActiveConnections.Empty();
+	}
+	Connections.Empty();
+}
+
+uint32 FLwsWebSocketServer::Run()
+{
+	// ~30Hz polling loop (poll mode — lws_service returns immediately)
+	const double TargetFrameTime = 1.0 / 30.0;
+
+	while (!ExitRequest.GetValue())
+	{
+		double BeginTime = FPlatformTime::Seconds();
+		if (LwsContext)
+		{
+			lws_service(LwsContext, 0);
+		}
+		double Elapsed = FPlatformTime::Seconds() - BeginTime;
+		double SleepTime = TargetFrameTime - Elapsed;
+		if (SleepTime > 0.0)
+		{
+			FPlatformProcess::SleepNoStats(static_cast<float>(SleepTime));
+		}
+	}
+	return 0;
+}
+
+void FLwsWebSocketServer::Stop()
+{
+	ExitRequest.Set(1);
+	WakeService();
+}
+
+void FLwsWebSocketServer::WakeService()
+{
+	if (LwsContext)
+	{
+		lws_cancel_service(LwsContext);
+	}
+}
+
+// ── Game-thread event enqueuers (called from lws thread) ────────────────────
+
+void FLwsWebSocketServer::LwsThread_OnEstablished(lws* Wsi, const FString& Path)
+{
+	TSharedPtr<FLwsWebSocketClientConnection> Conn = MakeShared<FLwsWebSocketClientConnection>(Wsi, Path, this);
+
+	// Store pointer in lws per-session data so callback can retrieve it on RECEIVE/WRITEABLE
+	FLwsWebSocketClientConnection** SessionData =
+		static_cast<FLwsWebSocketClientConnection**>(lws_wsi_user(Wsi));
+	if (SessionData)
+	{
+		*SessionData = Conn.Get();
+	}
+
+	// Register in thread-safe map so RECEIVE and CLOSED can look up the shared ptr
+	{
+		FScopeLock Lock(&ActiveConnectionsLock);
+		ActiveConnections.Add(Wsi, Conn);
+	}
+
+	ConnectedQueue.Enqueue({ Conn });
+}
+
+void FLwsWebSocketServer::LwsThread_OnReceive(lws* Wsi, const FString& Message)
+{
+	TSharedPtr<FLwsWebSocketClientConnection> Conn;
+	{
+		FScopeLock Lock(&ActiveConnectionsLock);
+		TSharedPtr<FLwsWebSocketClientConnection>* Found = ActiveConnections.Find(Wsi);
+		if (Found)
+		{
+			Conn = *Found;
+		}
+	}
+	if (Conn.IsValid())
+	{
+		UE_LOG(LogWebSocketServer, Display, TEXT("WS LwsThread_OnReceive: enqueuing msg='%s'"), *Message);
+		MessageQueue.Enqueue({ Conn, Message });
+	}
+	else
+	{
+		UE_LOG(LogWebSocketServer, Warning, TEXT("WS LwsThread_OnReceive: no Conn found for wsi!"));
+	}
+}
+
+void FLwsWebSocketServer::LwsThread_OnClosed(lws* Wsi)
+{
+	TSharedPtr<FLwsWebSocketClientConnection> Conn;
+	{
+		FScopeLock Lock(&ActiveConnectionsLock);
+		TSharedPtr<FLwsWebSocketClientConnection>* Found = ActiveConnections.Find(Wsi);
+		if (Found)
+		{
+			Conn = *Found;
+			ActiveConnections.Remove(Wsi);
+		}
+	}
+	if (Conn.IsValid())
+	{
+		DisconnectedQueue.Enqueue({ Conn });
+	}
+}
+
+// ── Game-thread Tick ─────────────────────────────────────────────────────────
+
+void FLwsWebSocketServer::Tick(float DeltaTime)
+{
+	// 1. Drain connected events first (so Connections is populated before message events)
+	{
+		FConnectedEvent Event;
+		while (ConnectedQueue.Dequeue(Event))
+		{
+			Connections.Add(Event.Conn);
+			if (ConnectedHandler)
+			{
+				ConnectedHandler(Event.Conn.ToSharedRef());
+			}
+		}
+	}
+
+	// 2. Drain message events
+	{
+		FMessageEvent Event;
+		while (MessageQueue.Dequeue(Event))
+		{
+			UE_LOG(LogWebSocketServer, Display, TEXT("WS Tick: draining message '%s'"), *Event.Message);
+			if (Event.Conn.IsValid() && MessageHandler)
+			{
+				MessageHandler(Event.Conn.ToSharedRef(), Event.Message);
+			}
+		}
+	}
+
+	// 3. Drain disconnected events
+	{
+		FDisconnectedEvent Event;
+		while (DisconnectedQueue.Dequeue(Event))
+		{
+			UE_LOG(LogWebSocketServer, Display, TEXT("WS Tick: draining disconnected event"));
+			if (Event.Conn.IsValid())
+			{
+				if (DisconnectedHandler)
+				{
+					DisconnectedHandler(Event.Conn.ToSharedRef());
+				}
+				Connections.RemoveSwap(Event.Conn);
+			}
+		}
+	}
+}
+
+// ── Static lws callback ──────────────────────────────────────────────────────
+
+int FLwsWebSocketServer::StaticLwsCallback(lws* Wsi, lws_callback_reasons Reason,
+                                            void* UserData, void* Data, size_t Length)
+{
+	// Retrieve server from context user data
+	lws_context* Context = lws_get_context(Wsi);
+	FLwsWebSocketServer* Self = static_cast<FLwsWebSocketServer*>(lws_context_user(Context));
+	if (!Self)
+	{
+		return 0;
+	}
+
+	// Per-session data holds FLwsWebSocketClientConnection* (after ESTABLISHED)
+	FLwsWebSocketClientConnection* Conn = nullptr;
+	if (UserData)
+	{
+		Conn = *static_cast<FLwsWebSocketClientConnection**>(UserData);
+	}
+
+	switch (Reason)
+	{
+	case LWS_CALLBACK_HTTP:
+		// Accept HTTP requests so lws doesn't 404 before upgrading to WebSocket
+		return 0;
+
+	case LWS_CALLBACK_FILTER_PROTOCOL_CONNECTION:
+		// Accept all incoming WebSocket connections regardless of Sec-WebSocket-Protocol
+		return 0;
+
+	case LWS_CALLBACK_ESTABLISHED:
+	{
+		// Read the request URI
+		char UriBuf[512];
+		FMemory::Memzero(UriBuf, sizeof(UriBuf));
+		lws_hdr_copy(Wsi, UriBuf, sizeof(UriBuf) - 1, WSI_TOKEN_GET_URI);
+		FString Path = UTF8_TO_TCHAR(UriBuf);
+
+		Self->LwsThread_OnEstablished(Wsi, Path);
+		break;
+	}
+
+	case LWS_CALLBACK_RECEIVE:
+	{
+		if (Data && Length > 0)
+		{
+			FUTF8ToTCHAR Convert(reinterpret_cast<const ANSICHAR*>(Data), static_cast<int32>(Length));
+			FString Message(Convert.Length(), Convert.Get());
+			Self->LwsThread_OnReceive(Wsi, Message);
+		}
+		break;
+	}
+
+	case LWS_CALLBACK_EVENT_WAIT_CANCELLED:
+	{
+		// WakeService() was called from another thread (send or close enqueued).
+		// Schedule SERVER_WRITEABLE for all live connections so they can flush/close.
+		FScopeLock Lock(&Self->ActiveConnectionsLock);
+		UE_LOG(LogWebSocketServer, Display, TEXT("WS EVENT_WAIT_CANCELLED: %d active connections"), Self->ActiveConnections.Num());
+		for (auto& Pair : Self->ActiveConnections)
+		{
+			lws_callback_on_writable(Pair.Key);
+		}
+		break;
+	}
+
+	case LWS_CALLBACK_SERVER_WRITEABLE:
+	{
+		UE_LOG(LogWebSocketServer, Display, TEXT("WS SERVER_WRITEABLE: Conn=%p"), Conn);
+		if (Conn)
+		{
+			if (Conn->LwsThreadFlushAndCheckClose())
+			{
+				UE_LOG(LogWebSocketServer, Display, TEXT("WS SERVER_WRITEABLE: returning -1 to close"));
+				return -1; // initiates lws close handshake
+			}
+		}
+		break;
+	}
+
+	case LWS_CALLBACK_CLOSED:
+	case LWS_CALLBACK_WSI_DESTROY:
+	{
+		UE_LOG(LogWebSocketServer, Display, TEXT("WS CLOSED/WSI_DESTROY: Reason=%d Conn=%p"), (int)Reason, Conn);
+		if (Conn)
+		{
+			Conn->LwsThreadOnClosed();
+			Self->LwsThread_OnClosed(Wsi);
+			// Clear per-session pointer so we don't double-fire
+			if (UserData)
+			{
+				*static_cast<FLwsWebSocketClientConnection**>(UserData) = nullptr;
+			}
+		}
+		break;
+	}
+
+	default:
+		break;
+	}
+
+	return 0;
+}
+
+#endif // WITH_WEBSOCKETSERVER
diff --git a/Engine/Source/Runtime/Online/WebSocketServer/Private/LwsWebSocketServer.h b/Engine/Source/Runtime/Online/WebSocketServer/Private/LwsWebSocketServer.h
new file mode 100644
--- /dev/null
+++ b/Engine/Source/Runtime/Online/WebSocketServer/Private/LwsWebSocketServer.h
@@ -0,0 +1,84 @@
+// Copyright Epic Games, Inc. All Rights Reserved.
+#pragma once
+
+#include "CoreMinimal.h"
+
+#if WITH_WEBSOCKETSERVER
+
+#include "IWebSocketServer.h"
+#include "LwsWebSocketClientConnection.h"
+#include "HAL/Runnable.h"
+#include "HAL/ThreadSafeCounter.h"
+#include "Containers/Queue.h"
+
+// ─────────────────────────────────────────────────────────────────────────────
+// Server — owns the lws_context and service thread
+// ─────────────────────────────────────────────────────────────────────────────
+class FLwsWebSocketServer final
+	: public IWebSocketServer
+	, public FRunnable
+{
+public:
+	explicit FLwsWebSocketServer(uint32 InPort);
+	virtual ~FLwsWebSocketServer();
+
+	// IWebSocketServer
+	virtual bool StartListening() override;
+	virtual void StopListening() override;
+	virtual bool IsListening() const override { return bListening; }
+	virtual void Tick(float DeltaTime) override;
+
+	virtual void OnConnected(FWsConnectedHandler Handler) override    { ConnectedHandler    = MoveTemp(Handler); }
+	virtual void OnMessage(FWsMessageHandler Handler) override        { MessageHandler      = MoveTemp(Handler); }
+	virtual void OnDisconnected(FWsDisconnectedHandler Handler) override { DisconnectedHandler = MoveTemp(Handler); }
+
+	// Called from lws thread (via StaticLwsCallback) to enqueue game-thread events
+	void LwsThread_OnEstablished(lws* Wsi, const FString& Path);
+	void LwsThread_OnReceive(lws* Wsi, const FString& Message);
+	void LwsThread_OnClosed(lws* Wsi);
+
+	// Wakes the lws service loop (safe to call from any thread)
+	void WakeService();
+
+private:
+	// FRunnable — runs on dedicated lws thread
+	virtual bool   Init() override  { return true; }
+	virtual uint32 Run() override;
+	virtual void   Stop() override;
+	virtual void   Exit() override  {}
+
+	static int StaticLwsCallback(lws* Wsi, lws_callback_reasons Reason,
+	                              void* UserData, void* Data, size_t Length);
+
+	uint32              Port;
+	lws_context*        LwsContext  = nullptr;
+	FRunnableThread*    Thread      = nullptr;
+	FThreadSafeCounter  ExitRequest;
+	bool                bListening  = false;
+
+	lws_protocols       LwsProtocol[2]; // protocol + zero terminator
+
+	// Game-thread event queues (MPSC: many lws connections → one game thread)
+	struct FConnectedEvent    { TSharedPtr<FLwsWebSocketClientConnection> Conn; };
+	struct FMessageEvent      { TSharedPtr<FLwsWebSocketClientConnection> Conn; FString Message; };
+	struct FDisconnectedEvent { TSharedPtr<FLwsWebSocketClientConnection> Conn; };
+
+	TQueue<FConnectedEvent,    EQueueMode::Mpsc> ConnectedQueue;
+	TQueue<FMessageEvent,      EQueueMode::Mpsc> MessageQueue;
+	TQueue<FDisconnectedEvent, EQueueMode::Mpsc> DisconnectedQueue;
+
+	// User-supplied handlers (set from game thread before StartListening)
+	FWsConnectedHandler    ConnectedHandler;
+	FWsMessageHandler      MessageHandler;
+	FWsDisconnectedHandler DisconnectedHandler;
+
+	// All live connections (game thread only — drained from ConnectedQueue/DisconnectedQueue in Tick)
+	TArray<TSharedPtr<FLwsWebSocketClientConnection>> Connections;
+
+	// Map from lws* → shared connection ptr, protected by a lock.
+	// Allows lws thread to enqueue message/disconnect events with a valid shared ptr.
+	FCriticalSection                                             ActiveConnectionsLock;
+	TMap<lws*, TSharedPtr<FLwsWebSocketClientConnection>>       ActiveConnections;
+};
+
+#endif // WITH_WEBSOCKETSERVER
diff --git a/Engine/Source/Runtime/Online/WebSocketServer/Private/WebSocketServerModule.cpp b/Engine/Source/Runtime/Online/WebSocketServer/Private/WebSocketServerModule.cpp
new file mode 100644
--- /dev/null
+++ b/Engine/Source/Runtime/Online/WebSocketServer/Private/WebSocketServerModule.cpp
@@ -0,0 +1,110 @@
+// Copyright Epic Games, Inc. All Rights Reserved.
+
+#include "WebSocketServerModule.h"
+#include "Modules/ModuleManager.h"
+
+#if WITH_WEBSOCKETSERVER
+#include "LwsWebSocketServer.h"
+#endif
+
+DEFINE_LOG_CATEGORY(LogWebSocketServer);
+
+IMPLEMENT_MODULE(FWebSocketServerModule, WebSocketServer);
+
+FWebSocketServerModule* FWebSocketServerModule::Singleton = nullptr;
+
+FWebSocketServerModule::FWebSocketServerModule()
+{
+}
+
+FWebSocketServerModule::~FWebSocketServerModule()
+{
+}
+
+void FWebSocketServerModule::StartupModule()
+{
+	Singleton = this;
+}
+
+void FWebSocketServerModule::ShutdownModule()
+{
+	StopAllServers();
+	Servers.Empty();
+	Singleton = nullptr;
+}
+
+bool FWebSocketServerModule::IsAvailable()
+{
+	return Singleton != nullptr;
+}
+
+FWebSocketServerModule& FWebSocketServerModule::Get()
+{
+	if (!Singleton)
+	{
+		check(IsInGameThread());
+		FModuleManager::LoadModuleChecked<FWebSocketServerModule>("WebSocketServer");
+	}
+	check(Singleton);
+	return *Singleton;
+}
+
+TSharedPtr<IWebSocketServer> FWebSocketServerModule::GetWebSocketServer(uint32 Port)
+{
+	TSharedPtr<IWebSocketServer>* Existing = Servers.Find(Port);
+	if (Existing)
+	{
+		return *Existing;
+	}
+
+#if WITH_WEBSOCKETSERVER
+	TSharedPtr<IWebSocketServer> NewServer = MakeShared<FLwsWebSocketServer>(Port);
+
+	if (bServersEnabled)
+	{
+		NewServer->StartListening();
+	}
+
+	Servers.Add(Port, NewServer);
+	return NewServer;
+#else
+	UE_LOG(LogWebSocketServer, Warning, TEXT("WebSocketServer is not supported on this platform."));
+	return nullptr;
+#endif
+}
+
+void FWebSocketServerModule::StartAllServers()
+{
+	bServersEnabled = true;
+	for (auto& Pair : Servers)
+	{
+		if (!Pair.Value->IsListening())
+		{
+			Pair.Value->StartListening();
+		}
+	}
+}
+
+void FWebSocketServerModule::StopAllServers()
+{
+	bServersEnabled = false;
+	for (auto& Pair : Servers)
+	{
+		if (Pair.Value->IsListening())
+		{
+			Pair.Value->StopListening();
+		}
+	}
+}
+
+bool FWebSocketServerModule::Tick(float DeltaTime)
+{
+	if (bServersEnabled)
+	{
+		for (auto& Pair : Servers)
+		{
+			Pair.Value->Tick(DeltaTime);
+		}
+	}
+	return true;
+}
diff --git a/Engine/Source/Runtime/Online/WebSocketServer/Public/IWebSocketClientConnection.h b/Engine/Source/Runtime/Online/WebSocketServer/Public/IWebSocketClientConnection.h
new file mode 100644
--- /dev/null
+++ b/Engine/Source/Runtime/Online/WebSocketServer/Public/IWebSocketClientConnection.h
@@ -0,0 +1,24 @@
+// Copyright Epic Games, Inc. All Rights Reserved.
+#pragma once
+
+#include "CoreMinimal.h"
+
+/**
+ * Handle to a single connected WebSocket client.
+ * May be held past disconnection; SendText/Close are no-ops after disconnect.
+ * All methods are thread-safe.
+ */
+class IWebSocketClientConnection
+{
+public:
+	virtual ~IWebSocketClientConnection() {}
+
+	/** Send a UTF-8 text frame to this client. */
+	virtual void SendText(const FString& Message) = 0;
+
+	/** Initiate a graceful close with the given status code. */
+	virtual void Close(int32 Code = 1000, const FString& Reason = FString()) = 0;
+
+	/** The URL path from the HTTP Upgrade request (e.g. "/webtests/websocketstests/echo/"). */
+	virtual FString GetPath() const = 0;
+};
diff --git a/Engine/Source/Runtime/Online/WebSocketServer/Public/IWebSocketServer.h b/Engine/Source/Runtime/Online/WebSocketServer/Public/IWebSocketServer.h
new file mode 100644
--- /dev/null
+++ b/Engine/Source/Runtime/Online/WebSocketServer/Public/IWebSocketServer.h
@@ -0,0 +1,43 @@
+// Copyright Epic Games, Inc. All Rights Reserved.
+#pragma once
+
+#include "CoreMinimal.h"
+#include "IWebSocketClientConnection.h"
+
+/** Called when a new client connects. */
+using FWsConnectedHandler    = TFunction<void(TSharedRef<IWebSocketClientConnection>)>;
+/** Called when a client sends a text message. */
+using FWsMessageHandler      = TFunction<void(TSharedRef<IWebSocketClientConnection>, const FString&)>;
+/** Called when a client disconnects (clean or error). */
+using FWsDisconnectedHandler = TFunction<void(TSharedRef<IWebSocketClientConnection>)>;
+
+/**
+ * A WebSocket server listening on a single port.
+ * Handlers are invoked on the game thread.
+ */
+class IWebSocketServer
+{
+public:
+	virtual ~IWebSocketServer() {}
+
+	/** Begin accepting connections. Returns false if binding failed. */
+	virtual bool StartListening() = 0;
+
+	/** Stop accepting new connections and close all existing ones. */
+	virtual void StopListening() = 0;
+
+	/** Returns true if currently listening. */
+	virtual bool IsListening() const = 0;
+
+	/** Called every frame on the game thread to dispatch pending events. */
+	virtual void Tick(float DeltaTime) = 0;
+
+	/** Register a handler for new connections. Replaces any previous handler. */
+	virtual void OnConnected(FWsConnectedHandler Handler) = 0;
+
+	/** Register a handler for incoming text messages. Replaces any previous handler. */
+	virtual void OnMessage(FWsMessageHandler Handler) = 0;
+
+	/** Register a handler for disconnections. Replaces any previous handler. */
+	virtual void OnDisconnected(FWsDisconnectedHandler Handler) = 0;
+};
diff --git a/Engine/Source/Runtime/Online/WebSocketServer/Public/WebSocketServerModule.h b/Engine/Source/Runtime/Online/WebSocketServer/Public/WebSocketServerModule.h
new file mode 100644
--- /dev/null
+++ b/Engine/Source/Runtime/Online/WebSocketServer/Public/WebSocketServerModule.h
@@ -0,0 +1,59 @@
+// Copyright Epic Games, Inc. All Rights Reserved.
+#pragma once
+
+#include "CoreMinimal.h"
+#include "Containers/Map.h"
+#include "Containers/Ticker.h"
+#include "Modules/ModuleInterface.h"
+#include "Modules/ModuleManager.h"
+#include "Templates/SharedPointer.h"
+#include "IWebSocketServer.h"
+
+DECLARE_LOG_CATEGORY_EXTERN(LogWebSocketServer, Log, All);
+
+/**
+ * Module that manages WebSocket servers.
+ * Usage:
+ *   TSharedPtr<IWebSocketServer> Server = FWebSocketServerModule::Get().GetWebSocketServer(8001);
+ *   Server->OnMessage([](TSharedRef<IWebSocketClientConnection> C, const FString& M){ C->SendText(M); });
+ *   FWebSocketServerModule::Get().StartAllServers();
+ */
+class FWebSocketServerModule
+	: public IModuleInterface
+	, public FTSTickerObjectBase
+{
+public:
+	FWebSocketServerModule();
+	virtual ~FWebSocketServerModule();
+
+	/** Returns true if the module is loaded and ready. */
+	static bool IsAvailable();
+
+	/** Singleton access. Loads the module if needed. */
+	static FWebSocketServerModule& Get();
+
+	/**
+	 * Returns (or creates) an IWebSocketServer for the given port.
+	 * Bind handlers before calling StartAllServers().
+	 */
+	TSharedPtr<IWebSocketServer> GetWebSocketServer(uint32 Port);
+
+	/** Start all servers (idempotent). */
+	void StartAllServers();
+
+	/** Stop all servers. */
+	void StopAllServers();
+
+	// IModuleInterface
+	virtual void StartupModule() override;
+	virtual void ShutdownModule() override;
+
+	// FTSTickerObjectBase — drives game-thread event dispatch
+	virtual bool Tick(float DeltaTime) override;
+
+private:
+	static FWebSocketServerModule* Singleton;
+
+	TMap<uint32, TSharedPtr<IWebSocketServer>> Servers;
+	bool bServersEnabled = false;
+};
diff --git a/Engine/Source/Runtime/Online/WebSocketServer/WebSocketServer.Build.cs b/Engine/Source/Runtime/Online/WebSocketServer/WebSocketServer.Build.cs
new file mode 100644
--- /dev/null
+++ b/Engine/Source/Runtime/Online/WebSocketServer/WebSocketServer.Build.cs
@@ -0,0 +1,34 @@
+// Copyright Epic Games, Inc. All Rights Reserved.
+
+using UnrealBuildTool;
+
+public class WebSocketServer : ModuleRules
+{
+	protected virtual bool bPlatformSupportsWebSocketServer
+	{
+		get
+		{
+			return
+				Target.Platform == UnrealTargetPlatform.Win64 ||
+				Target.Platform == UnrealTargetPlatform.Mac ||
+				Target.IsInPlatformGroup(UnrealPlatformGroup.Unix);
+		}
+	}
+
+	public WebSocketServer(ReadOnlyTargetRules Target) : base(Target)
+	{
+		PrivateDependencyModuleNames.AddRange(new string[]
+		{
+			"Core",
+		});
+
+		if (bPlatformSupportsWebSocketServer)
+		{
+			AddEngineThirdPartyPrivateStaticDependencies(Target, "OpenSSL", "libWebSockets", "zlib");
+			PrivateDependencyModuleNames.Add("SSL");
+		}
+
+		PublicDefinitions.Add("WEBSOCKETSERVER_PACKAGE=1");
+		PublicDefinitions.Add("WITH_WEBSOCKETSERVER=" + (bPlatformSupportsWebSocketServer ? "1" : "0"));
+	}
+}
